Add work-in-progress .clang-format file

diff --git a/Engine/Extras/clang-format/README.md b/Engine/Extras/clang-format/README.md
new file mode 100644
--- /dev/null
+++ b/Engine/Extras/clang-format/README.md
@@ -0,0 +1,3 @@
+This is a work-in-progress project to set up clang-format for the Unreal Engine
+codebase.
+
