Add powershell script to clang-format pending changes

diff --git a/Engine/Extras/clang-format/perforce-clang-format-diff.ps1 b/Engine/Extras/clang-format/perforce-clang-format-diff.ps1
new file mode 100644
--- /dev/null
+++ b/Engine/Extras/clang-format/perforce-clang-format-diff.ps1
@@ -0,0 +1,152 @@
+[CmdletBinding(SupportsShouldProcess)]
+param (
+    [Parameter(ValueFromPipeline = $true, Position = 0)]
+    [String]
+    $Path,
+
+    [Parameter()]
+    [string]
+    $Changelist
+)
+
+begin {
+    Set-StrictMode -Version 3
+    $FilesToFormat = @()
+}
+process {
+    if (![string]::IsNullOrEmpty($Path)) {
+        $FilesToFormat += $Path
+    }
+}
+end {
+    function ProcessZTag {
+        param (
+            [Parameter(ValueFromPipeline)]
+            [String]
+            $Line 
+        )
+        process {
+            switch -regex ($Line) {
+                "^\.\.\. depotFile (.*)" { $DepotFile = $Matches[1] }
+                "^\.\.\. type (.*)" { $Type = $Matches[1] }
+                "^\.\.\. action (.*)" { $Action = $Matches[1] }
+                "^\s*$" { 
+                    if ($Type -like "*text*") { 
+                        Write-Verbose "Found text file $DepotFile action $Action"
+                        [PSCustomObject]@{
+                            DepotPath = $DepotFile
+                            LocalPath = $null
+                            IsAdd     = $Action -like "add"
+                        }
+                    }
+                    else {
+                        Write-Verbose "Skipping non-text file $DepotFile"
+                    }
+                    $DepotFile = $null
+                    $Type = $null
+                    $Action = $null
+                }
+            }
+        }
+    }
+
+    # Filter files on the command line to those actually present in perforce and convert to depot paths 
+    if (0 -ne $FilesToFormat.Length) {
+        Write-Verbose "Checking perforce status of requested files $FilesToFormat"
+        $DepotFiles = p4 "-ztag" "fstat" "-Ro" @FilesToFormat | ProcessZTag
+    }
+    else {
+        if (![string]::IsNullOrEmpty($Changelist)) {
+            Write-Verbose "Requested files for changelist $Changelist"
+            # Get opened files in a specific changelist
+            $P4Args = @("-ztag", "opened", "-c", $Changelist);
+        }
+        else {
+            Write-Verbose "Requested files for all pending changelists"
+            # Get current client name
+            switch -regex (p4 "-ztag" "info") {
+                "^\.\.\. clientName (.*)" { $ClientName = $Matches[1] }
+            }
+            if ([string]::IsNullOrEmpty($ClientName)) {
+                throw "Unable to get perforce client name"
+            }
+            # Get all opened files
+            $P4Args = @("-ztag", "opened", "-C", $ClientName);
+        }
+
+        Write-Verbose "Fetching open files from Perforce $P4Args"
+        $DepotFiles = p4 @P4Args | ProcessZTag
+    }
+    
+    if ($DepotFiles.Length -eq 0) {
+        throw "No files to format"
+    }
+    $Index = 0
+    switch -regex (p4 -ztag "where" @($DepotFiles | Select-Object -ExpandProperty "DepotPath")) { 
+        "^\.\.\. path (.*)" {
+            $DepotFiles[$Index].LocalPath = $Matches[1]
+            ++$Index
+        }
+    }
+    
+    $Root = $PSScriptRoot
+    $DepotFiles | Foreach-Object -ThrottleLimit 5 -Parallel {
+        $VerbosePreference = $using:VerbosePreference
+        $WhatIfPreference = $using:WhatIfPreference
+        $Root = $using:Root
+        $LocalPath = $_.LocalPath
+        $IsAdd = $_.IsAdd
+        if ([string]::IsNullOrEmpty($LocalPath)) {
+            throw "Unexpected empty path"
+        }
+        if ($false -eq (Test-Path $LocalPath)) {
+            Write-Verbose "Skipping non-existed (deleted/moved?) file $LocalPath"
+            return;
+        }
+        
+        $clangFormatPath = Join-Path $Root "clang-format.exe"
+        $formatLocalPath = Join-Path $Root "experimental.clang-format"
+        # For newly added files, just format the entire file 
+        if ($IsAdd) {
+            if ($WhatIfPreference) {
+                Write-Host "Skipping full format of newly added file $LocalPath"
+            }
+            else {
+                & $clangFormatPath "-style=file:$formatLocalPath" "-i" $LocalPath | Out-Null 
+            }
+        }
+        else {
+            $clangFormatDiffPath = Join-Path $Root "clang-format-diff.py"
+
+            $diffOutput = p4 "diff" "-du" $LocalPath 
+        
+            # Check output is what we expect and bail if not 
+            if ($diffOutput.Length -lt 2) {
+                throw "Diff output for $LocalPath unexpectedly short"
+                $diffOutput | Write-Host
+            }
+            if ($diffOutput[0] -notlike "---*") {
+                throw "Missing old file path in diff output" 
+            }
+            if ($diffOutput[1] -notlike "+++*") {
+                throw "Missing new file path in diff output" 
+            }
+
+            $pythonArgs = @($clangFormatDiffPath, "-binary=$clangFormatPath", "-style=file:$formatLocalPath", "-v", "-i")
+
+            if ($WhatIfPreference) {
+                Write-Host "Skipping running diff script for file $LocalPath"
+                $diffOutput | Select-Object -first 2 | Write-Host
+            }
+            else {
+                $diffFormatOutput = Write-Output $diffOutput | python @pythonArgs | Out-String
+                $diffFormatExitCode = $LASTEXITCODE
+
+                if ($diffFormatExitCode -ne 0) {
+                    throw "Failed to run clang-format-diff.py on $LocalPath.`nclang-format-diff.py output:`n$diffFormatOutput"
+                }
+            }
+        }
+    }
+}
+
