Make sure list of files to format is an array

diff --git a/Engine/Extras/clang-format/perforce-clang-format-diff.ps1 b/Engine/Extras/clang-format/perforce-clang-format-diff.ps1
--- a/Engine/Extras/clang-format/perforce-clang-format-diff.ps1
+++ b/Engine/Extras/clang-format/perforce-clang-format-diff.ps1
@@ -53,7 +53,7 @@ end {
     # Filter files on the command line to those actually present in perforce and convert to depot paths 
     if (0 -ne $FilesToFormat.Length) {
         Write-Verbose "Checking perforce status of requested files $FilesToFormat"
-        $DepotFiles = p4 "-ztag" "fstat" "-Ro" @FilesToFormat | ProcessZTag
+        [array]$DepotFiles = p4 "-ztag" "fstat" "-Ro" @FilesToFormat | ProcessZTag
     }
     else {
         if (![string]::IsNullOrEmpty($Changelist)) {
@@ -75,7 +75,7 @@ end {
         }
 
         Write-Verbose "Fetching open files from Perforce $P4Args"
-        $DepotFiles = p4 @P4Args | ProcessZTag
+        [array]$DepotFiles = p4 @P4Args | ProcessZTag
     }
     
     if ($DepotFiles.Length -eq 0) {
