Reduce noise in clang-format-diff by skipping all context in Perforce diffs

This will avoid cases where clang-formats of only what has changed in Perforce still formatted unmodified code before or after the changes.

diff --git a/Engine/Extras/clang-format/perforce-clang-format-diff.ps1 b/Engine/Extras/clang-format/perforce-clang-format-diff.ps1
--- a/Engine/Extras/clang-format/perforce-clang-format-diff.ps1
+++ b/Engine/Extras/clang-format/perforce-clang-format-diff.ps1
@@ -118,7 +118,7 @@ end {
         else {
             $clangFormatDiffPath = Join-Path $Root "clang-format-diff.py"
 
-            $diffOutput = p4 "diff" "-du" $LocalPath 
+            $diffOutput = p4 "diff" "-du0" $LocalPath 
         
             # Check output is what we expect and bail if not 
             if ($diffOutput.Length -lt 2) {
