Improve our README.md for clang-format with sections on required binaries, what clang-format-diff is, and how to set up automatic diff-formatting in Rider

diff --git a/Engine/Extras/clang-format/README.md b/Engine/Extras/clang-format/README.md
--- a/Engine/Extras/clang-format/README.md
+++ b/Engine/Extras/clang-format/README.md
@@ -1,3 +1,40 @@
 This is a work-in-progress project to set up clang-format for the Unreal Engine
 codebase.
 
+# Required binaries
+
+Right now, you have to add clang-format.exe and clang-format-diff.py by
+yourself. Download them from an LLVM release and put them into this directory.
+
+# clang-format-diff
+
+Clang comes with a handy script called clang-format-diff which uses
+clang-format to format only the lines of code that has changed according to
+your version control system.
+
+To diff lines that have changed according to Perforce, you can use the script
+perforce-clang-format-diff.ps1.
+
+# Formatting changed lines on save in Rider
+
+You can format files automatically using clang-format-diff when you save
+changes to them from the Rider IDE.
+
+To do this, go to Rider's "Settings > Tools > File Watchers" and add a two new
+ones. Here are examples of values from a successful setup of these (the only
+that change between them is "File type"):
+
+| Field | Value |
+|---|---|
+| File type | C++ header files (or "C++ files" for .cpp files) |
+| Scope | Open Files |
+| Program | pwsh.exe |
+| Arguments | -File PATH-TO-CHECKOUT\Engine\Extras\clang-format\perforce-clang-format-diff.ps1 $FilePath$ |
+| Working directory | $FileDir$ |
+
+Note that the "Show console" option can be useful to debug the watcher.
+
+# Known issues
+
+clang-format currently does not format Slate code very well.
+
