Add -PerformSlateFiltering flag to perforce-clang-format-diff.ps1 to skip formatting of declarative Slate code

By default, this new filter is turned off. When enabled, this flag runs a filter that strips out any Slate-like blocks from the output of p4 diff before it's given to clang-format-diff.py. For the purposes of this filter, Slate-like code blocks start with a line containing only a [ and end in a line containing only a ] or ];

diff --git a/Engine/Extras/clang-format/perforce-clang-format-diff.ps1 b/Engine/Extras/clang-format/perforce-clang-format-diff.ps1
--- a/Engine/Extras/clang-format/perforce-clang-format-diff.ps1
+++ b/Engine/Extras/clang-format/perforce-clang-format-diff.ps1
@@ -6,7 +6,11 @@ param (
 
     [Parameter()]
     [string]
-    $Changelist
+    $Changelist,
+
+    [Parameter()]
+    [switch]
+    $PerformSlateFiltering
 )
 
 begin {
@@ -50,6 +54,118 @@ end {
         }
     }
 
+    function filterSlateCodeFromDiffOutput($diffOutput) {
+        Write-Verbose "Performing Slater filtering by skipping formatting on parts of the code that looks like declarative Slate code."
+
+        $fileContent = Get-Content $LocalPath
+        $fileContentLines = $fileContent -split "\r?\n|\r"
+
+        # Build zones.
+        $currentStartLineIndex = -1
+        $currentBracketNesting = 0
+        $slateBracketZones = for ($i=0; $i -lt $fileContentLines.Length; $i++)
+        {
+            $line = $fileContentLines[$i]
+            if ($line.Trim() -eq "[")
+            {
+                $currentBracketNesting += 1
+                # We don't record the start of a nested Slate bracket.
+                if ($currentBracketNesting -gt 0)
+                {
+                    $currentStartLineIndex = $i
+                }
+            }
+            elseif ($line.Trim() -eq "]" -or $line.Trim() -eq "];")
+            {
+                $currentBracketNesting -= 1
+                if ($currentBracketNesting -eq 0)
+                {
+                    $start = $currentStartLineIndex
+                    $currentStartLineIndex = -1
+                    @{ "startLineIndex" = [int]$start; "endLineIndex" = [int]$i };
+                }
+            }
+        }
+
+        # Expand zones up and down to the nearest whitespace-only lines.
+        foreach ($zone in $slateBracketZones)
+        {
+            $startIndex = $zone["startLineIndex"]
+            $endIndex = $zone["endLineIndex"]
+
+            $originalstartIndex = $startIndex
+            $originalEndIndex = $endIndex
+
+            # Keep moving the start of the zone up until we find a whitespace-only line or the start of the file.
+            while (($startIndex -gt 0) -and ($fileContentLines[$startIndex - 1].Trim() -ne ""))
+            {
+                $startIndex--
+            }
+
+            # Keep moving the end of the zone down until we find a whitespace-only line or the end of the file.
+            while (($endIndex -lt $fileContentLines.Length-1) -and ($fileContentLines[$endIndex + 1].Trim() -ne ""))
+            {
+                $endIndex++
+            }
+
+            if (($originalStartIndex -ne $startIndex) -or ($originalEndIndex -ne $endIndex))
+            {
+                Write-Verbose "Expanded Slate bracket zone from ($originalStartIndex, $originalEndIndex) to ($startIndex, $endIndex)"
+            }
+
+            $zone["startLineIndex"] = $startIndex
+            $zone["endLineIndex"] = $endIndex
+        }
+
+        $skipCurrentHunk = $false
+        $diffLines = $diffOutput -split "\r?\n|\r"
+        for ($i=0; $i -lt $diffLines.Length; $i++)
+        {
+            $diffLine = $diffLines[$i]
+            if ($diffLine -match "^@@.*\+(\d+)(?:,(\d+))?")
+            {
+                $skipCurrentHunk = $false
+
+                $match = $Matches
+
+                $hunkStartLine = [int]$match.1
+                $hunkLineCount = [int]$match.2
+                if (-Not $hunkLineCount)
+                {
+                    $hunkLineCount = 1
+                }
+                $hunkEndLine = $hunkStartLine + $hunkLineCount - 1
+
+                $hunkStartLineIndex = $hunkStartLine - 1
+                $hunkEndLineIndex = $hunkEndLine - 1
+
+                foreach ($slateBracketZone in $slateBracketZones)
+                {
+                    $zoneStart = $slateBracketZone["startLineIndex"]
+                    $zoneEnd = $slateBracketZone["endLineIndex"]
+
+                    $hunkStartInside = ($zoneStart -le $hunkStartLineIndex) -and ($hunkStartLineIndex -le $zoneEnd)
+                    $hunkEndInside = ($zoneStart -le $hunkEndLineIndex) -and ($hunkEndLineIndex -le $zoneEnd)
+
+                    if ($hunkStartInside -or $hunkEndInside)
+                    {
+                        Write-Verbose "Diff hunk overlaps Slate bracket zone, skipping the following lines:"
+                        $skipCurrentHunk = $true
+                    }
+                }
+            }
+
+            if ($skipCurrentHunk)
+            {
+                Write-Verbose "SKIPPING: $diffLine"
+            }
+            else
+            {
+                Write-Output $diffLine
+            }
+        }
+    }
+
     # Filter files on the command line to those actually present in perforce and convert to depot paths 
     if (0 -ne $FilesToFormat.Length) {
         Write-Verbose "Checking perforce status of requested files $FilesToFormat"
@@ -88,14 +204,20 @@ end {
             ++$Index
         }
     }
-    
+
+    # We need this to use the function from within Foreach-Object.
+    $filterSlateCodeFromDiffOutputSource = ${function:filterSlateCodeFromDiffOutput}.ToString()
+
     $Root = $PSScriptRoot
     $DepotFiles | Foreach-Object -ThrottleLimit 5 -Parallel {
         $VerbosePreference = $using:VerbosePreference
         $WhatIfPreference = $using:WhatIfPreference
+        $PerformSlateFiltering = $using:PerformSlateFiltering
+        ${function:filterSlateCodeFromDiffOutput} = $using:filterSlateCodeFromDiffOutputSource
         $Root = $using:Root
         $LocalPath = $_.LocalPath
         $IsAdd = $_.IsAdd
+
         if ([string]::IsNullOrEmpty($LocalPath)) {
             throw "Unexpected empty path"
         }
@@ -119,7 +241,15 @@ end {
             $clangFormatDiffPath = Join-Path $Root "clang-format-diff.py"
 
             $diffOutput = p4 "diff" "-du0" $LocalPath 
-        
+
+            if ($PerformSlateFiltering) {
+                $diffOutput = filterSlateCodeFromDiffOutput $diffOutput
+                Write-Verbose "Final diff given to clang-format-diff:"
+                foreach ($line in $diffOutput) {
+                    Write-Verbose $line
+                }
+            }
+
             # Check output is what we expect and bail if not 
             if ($diffOutput.Length -lt 2) {
                 throw "Diff output for $LocalPath unexpectedly short"
