Update our clang-format README.md with information about the new -PerformSlateFiltering flag

diff --git a/Engine/Extras/clang-format/README.md b/Engine/Extras/clang-format/README.md
--- a/Engine/Extras/clang-format/README.md
+++ b/Engine/Extras/clang-format/README.md
@@ -29,12 +29,46 @@ that change between them is "File type"):
 | File type | C++ header files (or "C++ files" for .cpp files) |
 | Scope | Open Files |
 | Program | pwsh.exe |
-| Arguments | -File PATH-TO-CHECKOUT\Engine\Extras\clang-format\perforce-clang-format-diff.ps1 $FilePath$ |
+| Arguments | -File PATH-TO-CHECKOUT\Engine\Extras\clang-format\perforce-clang-format-diff.ps1 $FilePath$ -PerformSlateFiltering |
 | Working directory | $FileDir$ |
 
+Note that this example enables the -PerformSlateFiltering flag (see below for info).
+
 Note that the "Show console" option can be useful to debug the watcher.
 
-# Known issues
+# Known issue: formatting Slate code
+
+Slate code uses an uncommon overload of operator[] to nest widgets
+declaratively. clang-format doesn't understand this and therefore cannot format
+such code nicely.
+
+An example of this problem is that this code
+
+```
+AddSlot()
+[
+    MyWidget
+];
+```
+
+is formatted into this
+
+```
+AddSlot()[MyWidget];
+```
+
+which is not what we want for our Slate code.
+
+To work around this issue, we can filter out any Slate code from what we give to
+clang-format. This works well when using clang-format-diff, which is already
+diff based.
 
-clang-format currently does not format Slate code very well.
+This technique has been implemented in the perforce-clang-format-diff.ps1 script
+as a filter on the output of `p4 diff` where any Slate-like code is removed from
+the diff before passed onto clang-format-diff.py. Enable this filter by passing
+`-PerformSlateFiltering` to perforce-clang-format-diff.ps1.
 
+Note that in rare cases, clang-format-diff will format code surrounding a diff
+regardless if the surrounding code has been changed or not. This can effectively
+override `-PerformSlateFiltering`, causing Slate code to be filtered
+incorrectly despite using the flag.
