Add perforce-clang-format.py for cross-platform clang-format-diff support

diff --git a/Engine/Extras/clang-format/README.md b/Engine/Extras/clang-format/README.md
--- a/Engine/Extras/clang-format/README.md
+++ b/Engine/Extras/clang-format/README.md
@@ -36,6 +36,18 @@ Note that this example enables the -PerformSlateFiltering flag (see below for in
 
 Note that the "Show console" option can be useful to debug the watcher.
 
+# macOS support
+
+The new `perforce-clang-format-diff.py` works on both Windows and macOS. You can set it up in Rider like this:
+
+| Field | Value                                                                                          |
+|---|------------------------------------------------------------------------------------------------|
+| File type | C++ header files (or "C++ files" for .cpp files)                                               |
+| Scope | Open Files                                                                                     |
+| Program | python                                                                                         |
+| Arguments | PATH-TO-CHECKOUT\Engine\Extras\clang-format\perforce-clang-format-diff.py $FilePath$ --verbose |
+| Working directory | $FileDir$                                                                                      |
+
 # Known issue: formatting Slate code
 
 Slate code uses an uncommon overload of operator[] to nest widgets
diff --git a/Engine/Extras/clang-format/perforce-clang-format-diff.py b/Engine/Extras/clang-format/perforce-clang-format-diff.py
new file mode 100644
--- /dev/null
+++ b/Engine/Extras/clang-format/perforce-clang-format-diff.py
@@ -0,0 +1,127 @@
+from dataclasses import dataclass
+import argparse
+import os
+import re
+import subprocess
+
+activate_print_verbose = False
+
+def print_verbose(*args, **kwargs):
+    if activate_print_verbose:
+        print("VERBOSE:", *args, **kwargs)
+
+def print_warning(*args, **kwargs):
+    print("WARNING:", *args, **kwargs)
+
+def decode_string(string):
+    if type(string) is str:
+        return string
+
+    try:
+        result = string.decode("utf-8")
+    except UnicodeDecodeError:
+        try:
+            result = string.decode("cp1252")
+        except UnicodeDecodeError:
+            print_warning("Failed to decode so returning original string \"{0}\"".format(string))
+            return string
+    
+    return result
+
+@dataclass
+class P4_ztag():
+    file_path: str = None
+    depot_file: str = None
+    is_add: bool = False
+
+def p4_ztag(path):
+    command = ["p4", "-ztag", "fstat", "-Ro", path]
+    result = subprocess.run(command, capture_output=True)
+    stdout_string = decode_string(result.stdout)
+    lines = stdout_string.splitlines()
+
+    depot_file_regex = re.compile(r'^\.\.\. depotFile (.*)')
+    type_regex = re.compile(r'^\.\.\. type (.*)')
+    action_regex = re.compile(r'^\.\.\. action (.*)')
+
+    result = P4_ztag(file_path = path)
+    for line in lines:
+        if match := depot_file_regex.match(line):
+            result.depot_file = match.group(1)
+        if match := type_regex.match(line):
+            file_type = match.group(1)
+            if file_type.strip() != "text":
+                return None
+        if match := action_regex.match(line):
+            action = match.group(1)
+            if action.strip() == "add":
+                result.is_add = True
+
+    return result
+
+def clang_format(exe_path, config_path, file_path):
+    command = [exe_path, "-style=file:%s" % config_path, "-i", file_path]
+    result = subprocess.run(command, capture_output=True)
+    if result.returncode != 0:
+        print("Non-zero return code {0} from {1}".format(result.returncode, " ".join(command)))
+        print("stderr:", result.stderr)
+        return
+    print(decode_string(result.stdout))
+
+def p4_diff(path):
+    command = ["p4", "diff", "-du0", path]
+    result = subprocess.run(command, capture_output=True)
+    if result.returncode != 0:
+        raise Exception("p4 diff failed: %s" % result.stderr)
+    stdout_string = decode_string(result.stdout)
+    return stdout_string
+
+def clang_format_diff(clang_format_diff_path, clang_format_path, config_path, diff):
+    command = ["python", clang_format_diff_path, "-binary=%s" % clang_format_path, "-style=file:%s" % config_path, "-v", "-i"]
+    result = subprocess.run(command, capture_output=True, text=True, input=diff)
+    if result.returncode != 0:
+        print("Non-zero return code {0} from {1}".format(result.returncode, " ".join(command)))
+        print("stderr:", result.stderr)
+        return
+
+def main():
+    args_parser = argparse.ArgumentParser(
+        prog="perforce-clang-format",
+        description="Uses Perforce to run clang-format on the diff of files you've changed",
+    )
+    args_parser.add_argument("paths", nargs="*")
+    args_parser.add_argument("-v", "--verbose", action="store_true")
+    args_parser.add_argument("--clang-format-path")
+    args_parser.add_argument("--clang-format-config-path")
+    args_parser.add_argument("--clang-format-diff-path")
+    args = args_parser.parse_args()
+
+    script_dir = os.path.dirname(os.path.realpath(__file__))
+
+    clang_format_path = args.clang_format_path
+    if clang_format_path is None:
+        clang_format_path = os.path.join(script_dir, "clang-format.exe")
+
+    config_path = args.clang_format_config_path
+    if config_path is None:
+        config_path = os.path.join(script_dir, "experimental.clang-format")
+
+    clang_format_diff_path = args.clang_format_diff_path
+    if clang_format_diff_path is None:
+        clang_format_diff_path = os.path.join(script_dir, "clang-format-diff.py")
+
+    for path in args.paths:
+        ztag = p4_ztag(path)
+        if ztag is None:
+            continue
+
+        if ztag.is_add:
+            print("Fully formatting new file \"%s\"." % ztag.file_path)
+            clang_format(clang_format_path, config_path, ztag.file_path)
+        else:
+            diff = p4_diff(ztag.file_path)
+            print("Formatting diff of file \"%s\"." % ztag.file_path)
+            clang_format_diff(clang_format_diff_path, clang_format_path, config_path, diff)
+
+if __name__ == "__main__":
+    main()
