perforce-clang-format-diff now formats files that aren't added to Perforce

diff --git a/Engine/Extras/clang-format/perforce-clang-format-diff.py b/Engine/Extras/clang-format/perforce-clang-format-diff.py
--- a/Engine/Extras/clang-format/perforce-clang-format-diff.py
+++ b/Engine/Extras/clang-format/perforce-clang-format-diff.py
@@ -33,21 +33,29 @@ class P4_ztag():
     file_path: str = None
     depot_file: str = None
     is_add: bool = False
+    is_open: bool = True
 
 def p4_ztag(path):
     command = ["p4", "-ztag", "fstat", "-Ro", path]
     result = subprocess.run(command, capture_output=True)
     stdout_string = decode_string(result.stdout)
+    stderr_string = decode_string(result.stderr)
+    
+    to_return = P4_ztag(file_path = path)
+    
+    if "not opened on this client" in stderr_string:
+        to_return.is_open = False
+        return to_return
+    
     lines = stdout_string.splitlines()
 
     depot_file_regex = re.compile(r'^\.\.\. depotFile (.*)')
     type_regex = re.compile(r'^\.\.\. type (.*)')
     action_regex = re.compile(r'^\.\.\. action (.*)')
 
-    result = P4_ztag(file_path = path)
     for line in lines:
         if match := depot_file_regex.match(line):
-            result.depot_file = match.group(1)
+            to_return.depot_file = match.group(1)
         if match := type_regex.match(line):
             file_type = match.group(1)
             if file_type.strip() != "text":
@@ -55,9 +63,9 @@ def p4_ztag(path):
         if match := action_regex.match(line):
             action = match.group(1)
             if action.strip() == "add":
-                result.is_add = True
+                to_return.is_add = True
 
-    return result
+    return to_return
 
 def clang_format(exe_path, config_path, file_path):
     command = [exe_path, "-style=file:%s" % config_path, "-i", file_path]
@@ -118,6 +126,9 @@ def main():
         if ztag.is_add:
             print("Fully formatting new file \"%s\"." % ztag.file_path)
             clang_format(clang_format_path, config_path, ztag.file_path)
+        elif not ztag.is_open:
+            print("Fully formatting unopened file \"%s\"." % ztag.file_path)
+            clang_format(clang_format_path, config_path, ztag.file_path)
         else:
             diff = p4_diff(ztag.file_path)
             print("Formatting diff of file \"%s\"." % ztag.file_path)
