Make perforce-clang-format-diff.py open unopened files and then diff-format them. Previously, unopened files were fully formatted, regardless of the diff.

diff --git a/Engine/Extras/clang-format/perforce-clang-format-diff.py b/Engine/Extras/clang-format/perforce-clang-format-diff.py
--- a/Engine/Extras/clang-format/perforce-clang-format-diff.py
+++ b/Engine/Extras/clang-format/perforce-clang-format-diff.py
@@ -84,6 +84,12 @@ def p4_diff(path):
     stdout_string = decode_string(result.stdout)
     return stdout_string
 
+def p4_open(path):
+    command = ["p4", "open", path]
+    result = subprocess.run(command, capture_output=True)
+    if result.returncode != 0:
+        raise Exception("p4 open failed: %s" % result.stderr)
+
 def clang_format_diff(clang_format_diff_path, clang_format_path, config_path, diff):
     command = ["python", clang_format_diff_path, "-binary=%s" % clang_format_path, "-style=file:%s" % config_path, "-v", "-i"]
     result = subprocess.run(command, capture_output=True, text=True, input=diff)
@@ -127,8 +133,11 @@ def main():
             print("Fully formatting new file \"%s\"." % ztag.file_path)
             clang_format(clang_format_path, config_path, ztag.file_path)
         elif not ztag.is_open:
-            print("Fully formatting unopened file \"%s\"." % ztag.file_path)
-            clang_format(clang_format_path, config_path, ztag.file_path)
+            print("Opening unopened file \"%s\"." % ztag.file_path)
+            p4_open(ztag.file_path)
+            diff = p4_diff(ztag.file_path)
+            print("Formatting diff of file \"%s\"." % ztag.file_path)
+            clang_format_diff(clang_format_diff_path, clang_format_path, config_path, diff)
         else:
             diff = p4_diff(ztag.file_path)
             print("Formatting diff of file \"%s\"." % ztag.file_path)
