Add --verify flag to perforce-clang-format-diff.py

When using the new --verify flag, the files will not be modified and the script returns a non-zero error code when files needed formatting. This is useful to verify whether formatting was needed or not.

Also add support for finding Mac binaries of clang-format.

diff --git a/Engine/Extras/clang-format/perforce-clang-format-diff.py b/Engine/Extras/clang-format/perforce-clang-format-diff.py
--- a/Engine/Extras/clang-format/perforce-clang-format-diff.py
+++ b/Engine/Extras/clang-format/perforce-clang-format-diff.py
@@ -1,8 +1,10 @@
 from dataclasses import dataclass
 import argparse
 import os
+import platform
 import re
 import subprocess
+import sys
 
 activate_print_verbose = False
 
@@ -67,15 +69,25 @@ def p4_ztag(path):
 
     return to_return
 
-def clang_format(exe_path, config_path, file_path):
+def clang_format(exe_path, config_path, file_path, verify=False):
+    """Returns True if the diff needed formatting, False otherwise."""
+
     command = [exe_path, "-style=file:%s" % config_path, "-i", file_path]
+    if verify:
+        command.append("--dry-run")
+        command.append("-Werror")
+
     result = subprocess.run(command, capture_output=True)
+    if result.returncode == 1:
+        return True
     if result.returncode != 0:
         print("Non-zero return code {0} from {1}".format(result.returncode, " ".join(command)))
         print("stderr:", result.stderr)
         return
     print(decode_string(result.stdout))
 
+    return False
+
 def p4_diff(path):
     command = ["p4", "diff", "-du0", path]
     result = subprocess.run(command, capture_output=True)
@@ -90,13 +102,24 @@ def p4_open(path):
     if result.returncode != 0:
         raise Exception("p4 open failed: %s" % result.stderr)
 
-def clang_format_diff(clang_format_diff_path, clang_format_path, config_path, diff):
-    command = ["python", clang_format_diff_path, "-binary=%s" % clang_format_path, "-style=file:%s" % config_path, "-v", "-i"]
+def clang_format_diff(clang_format_diff_path, clang_format_path, config_path, diff, verify=False):
+    """Returns True if the diff needed formatting, False otherwise."""
+
+    command = ["python", clang_format_diff_path, "-binary=%s" % clang_format_path, "-style=file:%s" % config_path, "-v"]
+
+    if not verify:
+        # Write the formatted code back to the files being formatted.
+        command.append("-i")
+
     result = subprocess.run(command, capture_output=True, text=True, input=diff)
-    if result.returncode != 0:
+    # clang-format-diff returns 1 when the file needs to be formatted.
+    if result.returncode == 1:
+        return True
+    elif result.returncode != 0:
         print("Non-zero return code {0} from {1}".format(result.returncode, " ".join(command)))
         print("stderr:", result.stderr)
-        return
+
+    return False
 
 def main():
     args_parser = argparse.ArgumentParser(
@@ -108,13 +131,28 @@ def main():
     args_parser.add_argument("--clang-format-path")
     args_parser.add_argument("--clang-format-config-path")
     args_parser.add_argument("--clang-format-diff-path")
+    args_parser.add_argument("--verify", action="store_true")
     args = args_parser.parse_args()
 
+    global activate_print_verbose
+    activate_print_verbose = args.verbose
+
+    print_verbose("args.verify=", args.verify)
+    print_verbose("args.paths=", args.paths)
+
     script_dir = os.path.dirname(os.path.realpath(__file__))
 
     clang_format_path = args.clang_format_path
     if clang_format_path is None:
-        clang_format_path = os.path.join(script_dir, "clang-format.exe")
+        if platform.system() == "Darwin":
+            clang_format_path = os.path.join(script_dir, "Mac-arm64", "clang-format")
+        elif platform.system() == "Windows":
+            clang_format_path = os.path.join(script_dir, "Win64", "clang-format.exe")
+        elif platform.system() == "Linux":
+            raise Exception("Linux is not yet supported.")
+        else:
+            # Support the old clang-format.exe location for now.
+            clang_format_path = os.path.join(script_dir, "clang-format.exe")
 
     config_path = args.clang_format_config_path
     if config_path is None:
@@ -124,24 +162,50 @@ def main():
     if clang_format_diff_path is None:
         clang_format_diff_path = os.path.join(script_dir, "clang-format-diff.py")
 
+    files_that_need_formatting = []
+
     for path in args.paths:
         ztag = p4_ztag(path)
         if ztag is None:
             continue
 
         if ztag.is_add:
-            print("Fully formatting new file \"%s\"." % ztag.file_path)
-            clang_format(clang_format_path, config_path, ztag.file_path)
+            if args.verify:
+                print_verbose("Checking formatting of new file \"%s\"." % ztag.file_path)
+            else:
+                print_verbose("Fully formatting new file \"%s\"." % ztag.file_path)
+
+            if clang_format(clang_format_path, config_path, ztag.file_path, args.verify):
+                if args.verify:
+                    files_that_need_formatting.append(ztag.file_path)
+
         elif not ztag.is_open:
-            print("Opening unopened file \"%s\"." % ztag.file_path)
-            p4_open(ztag.file_path)
+            if not args.verify:
+                print_verbose("Opening unopened file \"%s\"." % ztag.file_path)
+                p4_open(ztag.file_path)
             diff = p4_diff(ztag.file_path)
-            print("Formatting diff of file \"%s\"." % ztag.file_path)
-            clang_format_diff(clang_format_diff_path, clang_format_path, config_path, diff)
+
+            if args.verify:
+                print_verbose("Checking formatting of diff of file \"%s\"." % ztag.file_path)
+            else:
+                print_verbose("Formatting diff of file \"%s\"." % ztag.file_path)
+
+            if clang_format_diff(clang_format_diff_path, clang_format_path, config_path, diff, args.verify):
+                if args.verify:
+                    files_that_need_formatting.append(ztag.file_path)
+
         else:
             diff = p4_diff(ztag.file_path)
-            print("Formatting diff of file \"%s\"." % ztag.file_path)
-            clang_format_diff(clang_format_diff_path, clang_format_path, config_path, diff)
+            print_verbose("Formatting diff of file \"%s\"." % ztag.file_path)
+
+            if clang_format_diff(clang_format_diff_path, clang_format_path, config_path, diff, args.verify):
+                if args.verify:
+                    files_that_need_formatting.append(ztag.file_path)
+
+    if len(files_that_need_formatting) > 0:
+        for file_path in files_that_need_formatting:
+            print("File needs formatting: \"%s\"." % file_path)
+        sys.exit(1)
 
 if __name__ == "__main__":
     main()
