clang-format-diff: Add additional paths when running commands on Mac so the p4 command can be found

diff --git a/Engine/Extras/clang-format/perforce-clang-format-diff.py b/Engine/Extras/clang-format/perforce-clang-format-diff.py
--- a/Engine/Extras/clang-format/perforce-clang-format-diff.py
+++ b/Engine/Extras/clang-format/perforce-clang-format-diff.py
@@ -3,6 +3,7 @@ import argparse
 import os
 import platform
 import re
+import shutil
 import subprocess
 import sys
 
@@ -30,6 +31,20 @@ def decode_string(string):
     
     return result
 
+def run_subprocess(command, **kwargs):
+    environment = os.environ.copy()
+
+    # Add additional paths on the Mac so it can find p4 etc.
+    if platform.system() == "Darwin":
+        additional_paths = ["/usr/local/bin", "/usr/bin", "/opt/homebrew/bin"]
+        new_path = os.environ.get("PATH", "")
+        for path in additional_paths:
+            new_path = os.pathsep.join([new_path, path])
+
+        environment = environment | {"PATH": new_path}
+
+    return subprocess.run(command, capture_output=True, env=environment, **kwargs)
+
 @dataclass
 class P4_ztag():
     file_path: str = None
@@ -39,7 +54,7 @@ class P4_ztag():
 
 def p4_ztag(path):
     command = ["p4", "-ztag", "fstat", "-Ro", path]
-    result = subprocess.run(command, capture_output=True)
+    result = run_subprocess(command)
     stdout_string = decode_string(result.stdout)
     stderr_string = decode_string(result.stderr)
     
@@ -77,7 +92,7 @@ def clang_format(exe_path, config_path, file_path, verify=False):
         command.append("--dry-run")
         command.append("-Werror")
 
-    result = subprocess.run(command, capture_output=True)
+    result = run_subprocess(command)
     if result.returncode == 1:
         return True
     if result.returncode != 0:
@@ -90,7 +105,7 @@ def clang_format(exe_path, config_path, file_path, verify=False):
 
 def p4_diff(path):
     command = ["p4", "diff", "-du0", path]
-    result = subprocess.run(command, capture_output=True)
+    result = run_subprocess(command)
     if result.returncode != 0:
         raise Exception("p4 diff failed: %s" % result.stderr)
     stdout_string = decode_string(result.stdout)
@@ -98,20 +113,31 @@ def p4_diff(path):
 
 def p4_open(path):
     command = ["p4", "open", path]
-    result = subprocess.run(command, capture_output=True)
+    result = run_subprocess(command)
     if result.returncode != 0:
         raise Exception("p4 open failed: %s" % result.stderr)
 
 def clang_format_diff(clang_format_diff_path, clang_format_path, config_path, diff, verify=False):
     """Returns True if the diff needed formatting, False otherwise."""
 
-    command = ["python", clang_format_diff_path, "-binary=%s" % clang_format_path, "-style=file:%s" % config_path, "-v"]
+    # Figure out which python executable to call.
+    python_command = None
+    if shutil.which("python") is not None:
+        python_command = "python"
+    elif shutil.which("python3") is not None:
+        python_command = "python3"
+    else:
+        raise Exception("Could not run clang-format-diff.py: could not find \"python\" or \"python3\" executable to run it with.")
+
+    print("Using python command \"%s\"." % python_command)
+
+    command = [python_command, clang_format_diff_path, "-binary=%s" % clang_format_path, "-style=file:%s" % config_path, "-v"]
 
     if not verify:
         # Write the formatted code back to the files being formatted.
         command.append("-i")
 
-    result = subprocess.run(command, capture_output=True, text=True, input=diff)
+    result = run_subprocess(command, text=True, input=diff)
     # clang-format-diff returns 1 when the file needs to be formatted.
     if result.returncode == 1:
         return True
