Make perforce-clang-format-diff use print_verbose for the chosen python command

diff --git a/Engine/Extras/clang-format/perforce-clang-format-diff.py b/Engine/Extras/clang-format/perforce-clang-format-diff.py
--- a/Engine/Extras/clang-format/perforce-clang-format-diff.py
+++ b/Engine/Extras/clang-format/perforce-clang-format-diff.py
@@ -129,7 +129,7 @@ def clang_format_diff(clang_format_diff_path, clang_format_path, config_path, di
     else:
         raise Exception("Could not run clang-format-diff.py: could not find \"python\" or \"python3\" executable to run it with.")
 
-    print("Using python command \"%s\"." % python_command)
+    print_verbose("Using python command \"%s\"." % python_command)
 
     command = [python_command, clang_format_diff_path, "-binary=%s" % clang_format_path, "-style=file:%s" % config_path, "-v"]
 
