clang-format wrapper

diff --git a/Engine/Extras/clang-format/experimental-clang-format.py b/Engine/Extras/clang-format/experimental-clang-format.py
new file mode 100644
--- /dev/null
+++ b/Engine/Extras/clang-format/experimental-clang-format.py
@@ -0,0 +1,131 @@
+#!/usr/bin/env python3
+import argparse, os, re, subprocess, sys
+
+# Helpers
+
+def run(cmd, **kw):
+    """Run subprocess, capture output, raise on error."""
+    return subprocess.run(cmd, stdout=subprocess.PIPE, stderr=subprocess.PIPE, text=True, check=True, **kw)
+
+def get_opened_depots():
+    out = run(["p4", "opened"]).stdout.splitlines()
+    # each line: //depot/.../File.cpp#rev - edit change ...
+    return [line.split()[0].split("#")[0] for line in out if line]
+
+def depot_to_client(path):
+    # p4 where //depot/... -> "//depot/... //client/... C:/.../File.cpp"
+    out = run(["p4", "where", path]).stdout.split()
+    return out[2] if len(out) >= 3 else None
+
+def build_diff(path):
+    return run(["p4", "diff", "-du0", path]).stdout
+
+# Main
+
+def main():
+    p = argparse.ArgumentParser(
+        description="Simple Perforce wrapper for clang-format-diff.py with verbose output"
+    )
+    p.add_argument("files", nargs="*", help="Depot or local paths to format")
+    p.add_argument("--verify", action="store_true",
+                   help="Dry-run: report files needing formatting")
+    p.add_argument("--diff-script", default="../../Engine/Extras/clang-format/clang-format-diff.py",
+                   help="Path to LLVM clang-format-diff.py")
+    p.add_argument("--clang-format", default="clang-format",
+                   help="Path to clang-format binary")
+    p.add_argument("-c", "--config-file", default=None,
+                   help="Explicit path to .clang-format to use")
+    p.add_argument("-v", "--verbose", action="store_true",
+                   help="Show detailed operations and diffs")
+    args = p.parse_args()
+    verbose = args.verbose
+
+    # determine list of targets
+    if args.files:
+        depots = [f.split("#")[0] for f in args.files]
+    else:
+        depots = get_opened_depots()
+        if verbose:
+            print("Opened files:")
+            for c in depots:
+                print(f"  - {c}")
+
+    # filter .cpp/.h
+    clients = []
+    for d in depots:
+        c = depot_to_client(d)
+        if c and re.search(r'\.(cpp|h)$', c, re.IGNORECASE):
+            clients.append(c)
+    if not clients:
+        print("No .cpp/.h files to format.")
+        return 0
+    if verbose:
+        print("Formatting .h/.cpp:")
+        for c in clients:
+            print(f"  - {c}")
+
+    status = 0
+    for local in clients:
+        diff = build_diff(local)
+        if not diff:
+            if verbose:
+                print(f"No changes in: {local}")
+            continue  # no changes
+
+        if verbose:
+            print(f"\n=== Processing: {local} ===")
+            print("--- original diff ---")
+            print(diff)
+
+        cmd = [
+            sys.executable, args.diff_script,
+            f"-binary={args.clang_format}",
+            "-style=file",
+        ]
+
+        # Build style argument: use explicit config when provided, else local .clang-format discovery
+        style_arg = "-style=file"
+        if args.config_file:
+            cfg = os.path.abspath(args.config_file)
+            style_arg = f"-style=file:{cfg}"
+            if args.verbose:
+                print(f"Using clang-format config: {cfg}")
+
+        cmd = [
+            sys.executable, args.diff_script,
+            f"-binary={args.clang_format}",
+            style_arg,
+        ]
+
+        if not args.verify:
+             cmd.append("-i")
+        if verbose:
+            print("Running:", ' '.join(cmd))
+
+        proc = subprocess.run(cmd, input=diff, text=True,
+                              stdout=subprocess.PIPE,
+                              stderr=subprocess.PIPE)
+
+        if verbose and proc.stdout:
+            print("--- clang-format-diff output ---")
+            print(proc.stdout)
+        if verbose and proc.stderr:
+            print("--- errors ---", file=sys.stderr)
+            print(proc.stderr, file=sys.stderr)
+
+        if args.verify:
+            print(f"*** NEEDS FORMATTING: {local}")
+            status = 1
+        else:
+            if proc.returncode == 0:
+                print(f"Patched diff-hunks in: {local}")
+            else:
+                # covers proc.returncode != 0
+                print(f"Error running clang-format-diff on {local}", file=sys.stderr)
+                print(proc.stderr, file=sys.stderr)
+                status = max(status, proc.returncode)
+
+    return status
+
+if __name__ == "__main__":
+    sys.exit(main())
